inital add of Shallow Water plugin

diff --git a/Engine/Plugins/Experimental/ShallowWater/ShallowWater.uplugin b/Engine/Plugins/Experimental/ShallowWater/ShallowWater.uplugin
new file mode 100644
--- /dev/null
+++ b/Engine/Plugins/Experimental/ShallowWater/ShallowWater.uplugin
@@ -0,0 +1,23 @@
+{
+	"FileVersion": 1,
+	"Version": 1,
+	"VersionName": "1.0",
+	"FriendlyName": "Shallow Water",
+	"Description": "Shallow Water Simulation",
+	"Category": "Physics",
+	"CreatedBy": "Epic Games",
+	"CreatedByURL": "epicgames.com",
+	"DocsURL": "",
+	"MarketplaceURL": "",
+	"SupportURL": "",
+	"EnabledByDefault": false,
+	"CanContainContent": true,
+	"IsBetaVersion": true,
+	"Installed": false,
+	"Plugins": [
+		{
+			"Name": "Niagara",
+			"Enabled": true
+		}
+	]
+}
\ No newline at end of file
